- Handling non-power of two errors in the cooker with an "error texture" with a message of what to do to the texture [UE-8923]

diff --git a/Engine/Source/Developer/TextureFormatUncompressed/Private/TextureFormatUncompressed.cpp b/Engine/Source/Developer/TextureFormatUncompressed/Private/TextureFormatUncompressed.cpp
--- a/Engine/Source/Developer/TextureFormatUncompressed/Private/TextureFormatUncompressed.cpp
+++ b/Engine/Source/Developer/TextureFormatUncompressed/Private/TextureFormatUncompressed.cpp
@@ -19,7 +19,8 @@ DEFINE_LOG_CATEGORY_STATIC(LogTextureFormatUncompressed, Log, All);
 	op(VU8) \
 	op(RGBA16F) \
 	op(XGXR8) \
-	op(RGBA8)
+	op(RGBA8) \
+	op(POTERROR)
 
 #define DECL_FORMAT_NAME(FormatName) static FName GTextureFormatName##FormatName = FName(TEXT(#FormatName));
 ENUM_SUPPORTED_FORMATS(DECL_FORMAT_NAME);
@@ -181,6 +182,39 @@ class FTextureFormatUncompressed : public ITextureFormat
 			OutCompressedImage.PixelFormat = PF_FloatRGBA;
 			OutCompressedImage.RawData = Image.RawData;
 
+			return true;
+		}
+		else if (BuildSettings.TextureFormatName == GTextureFormatNamePOTERROR)
+		{
+			// load the error image data we will just repeat into the texture
+			TArray<uint8> ErrorData;
+			FFileHelper::LoadFileToArray(ErrorData, TEXT("../../../Engine/Content/MobileResources/PowerOfTwoError64x64.raw"));
+
+			// set output
+			OutCompressedImage.SizeX = InImage.SizeX;
+			OutCompressedImage.SizeY = InImage.SizeY;
+			OutCompressedImage.PixelFormat = PF_B8G8R8A8;
+
+			// allocate output memory
+			check(InImage.NumSlices == 1);
+			uint32 NumTexels = InImage.SizeX * InImage.SizeY;
+			OutCompressedImage.RawData.Empty(NumTexels * 4);
+			OutCompressedImage.RawData.AddUninitialized(NumTexels * 4);
+
+			// write out texels
+			uint8* Src = ErrorData.GetData();
+			uint8* Dest = (uint8*)OutCompressedImage.RawData.GetData();
+			for (int32 Y = 0; Y < InImage.SizeY; Y++)
+			{
+				for (int32 X = 0; X < InImage.SizeX * 4; X++)
+				{
+					int32 SrcX = X & (64 * 4 - 1);
+					int32 SrcY = Y & 63;
+					Dest[Y * InImage.SizeX * 4 + X] = Src[SrcY * 64 * 4 + SrcX];
+				}
+			}
+
+			
 			return true;
 		}
 		
diff --git a/Engine/Source/Developer/iOS/IOSTargetPlatform/Private/IOSTargetPlatform.cpp b/Engine/Source/Developer/iOS/IOSTargetPlatform/Private/IOSTargetPlatform.cpp
--- a/Engine/Source/Developer/iOS/IOSTargetPlatform/Private/IOSTargetPlatform.cpp
+++ b/Engine/Source/Developer/iOS/IOSTargetPlatform/Private/IOSTargetPlatform.cpp
@@ -400,6 +400,8 @@ void FIOSTargetPlatform::GetTextureFormats( const UTexture* Texture, TArray<FNam
 		FName(TEXT("BC5")),		FName(TEXT("PVRTCN")),		FName(TEXT("ASTC_NormalRG")),
 		FName(TEXT("AutoDXT")),	FName(TEXT("AutoPVRTC")),	FName(TEXT("ASTC_RGBAuto")),
 	};
+	static FName NameBGRA8(TEXT("BGRA8"));
+	static FName NamePOTERROR(TEXT("POTERROR"));
 
 	FName TextureFormatName = NAME_None;
 
@@ -432,7 +434,20 @@ void FIOSTargetPlatform::GetTextureFormats( const UTexture* Texture, TArray<FNam
 			}
 			if (bIncludePVRTC)
 			{
-				OutFormats.AddUnique(FormatRemap[RemapIndex + 1]);
+#if WITH_EDITORONLY_DATA
+				// handle non-power of 2 textures
+				if (!FMath::IsPowerOfTwo(Texture->Source.GetSizeX()) || !FMath::IsPowerOfTwo(Texture->Source.GetSizeY()))
+				{
+					// option 1: Uncompress, but users will get very large textures unknowningly
+					// OutFormats.AddUnique(NameBGRA8);
+					// option 2: Use an "error message" texture so they see it in game
+					OutFormats.AddUnique(NamePOTERROR);
+				}
+				else
+#endif
+				{
+					OutFormats.AddUnique(FormatRemap[RemapIndex + 1]);
+				}
 			}
 		}
 	}
