Fixed crash loading non-power-of-two PVRTC fallback data when using a project not under the same tree as Engine.

diff --git a/Engine/Source/Developer/TextureFormatUncompressed/Private/TextureFormatUncompressed.cpp b/Engine/Source/Developer/TextureFormatUncompressed/Private/TextureFormatUncompressed.cpp
--- a/Engine/Source/Developer/TextureFormatUncompressed/Private/TextureFormatUncompressed.cpp
+++ b/Engine/Source/Developer/TextureFormatUncompressed/Private/TextureFormatUncompressed.cpp
@@ -212,7 +212,8 @@ class FTextureFormatUncompressed : public ITextureFormat
 		{
 			// load the error image data we will just repeat into the texture
 			TArray<uint8> ErrorData;
-			FFileHelper::LoadFileToArray(ErrorData, TEXT("../../../Engine/Content/MobileResources/PowerOfTwoError64x64.raw"));
+			FFileHelper::LoadFileToArray(ErrorData, *(FPaths::EngineDir() / TEXT("Content/MobileResources/PowerOfTwoError64x64.raw")));
+			check(ErrorData.Num() == 16384);
 
 			// set output
 			OutCompressedImage.SizeX = InImage.SizeX;
