Fix CIS

diff --git a/Engine/Source/Developer/TextureFormatUncompressed/Private/TextureFormatUncompressed.cpp b/Engine/Source/Developer/TextureFormatUncompressed/Private/TextureFormatUncompressed.cpp
--- a/Engine/Source/Developer/TextureFormatUncompressed/Private/TextureFormatUncompressed.cpp
+++ b/Engine/Source/Developer/TextureFormatUncompressed/Private/TextureFormatUncompressed.cpp
@@ -2,6 +2,7 @@
 
 #include "CoreMinimal.h"
 #include "Misc/FileHelper.h"
+#include "Misc/Paths.h"
 #include "Modules/ModuleManager.h"
 #include "Interfaces/ITextureFormat.h"
 #include "Interfaces/ITextureFormatModule.h"
