Change DDC key for uncompressed 16F textures to force recook
in case they got cooked during F32-F16 signed float bug week
fixes tree proxy impostors

diff --git a/Engine/Source/Developer/TextureFormatUncompressed/Private/TextureFormatUncompressed.cpp b/Engine/Source/Developer/TextureFormatUncompressed/Private/TextureFormatUncompressed.cpp
--- a/Engine/Source/Developer/TextureFormatUncompressed/Private/TextureFormatUncompressed.cpp
+++ b/Engine/Source/Developer/TextureFormatUncompressed/Private/TextureFormatUncompressed.cpp
@@ -51,11 +51,32 @@ class FTextureFormatUncompressed : public ITextureFormat
 
 	virtual uint16 GetVersion(
 		FName Format,
-		const struct FTextureBuildSettings* BuildSettings = nullptr
+		const struct FTextureBuildSettings* BuildSettings
 	) const override
 	{
 		return 0;
 	}
+	
+	virtual FString GetDerivedDataKeyString(const class UTexture& InTexture, const struct FTextureBuildSettings* InBuildSettings) const override
+	{
+		check( InBuildSettings != NULL );
+		
+		if (InBuildSettings->TextureFormatName == GTextureFormatNameRGBA16F)
+		{
+			return TEXT("RGBA16F");
+		}
+		else if (InBuildSettings->TextureFormatName == GTextureFormatNameR16F)
+		{
+			return TEXT("R16F");
+		}
+		else
+		{
+			// default implemenation of GetDerivedDataKeyString returns empty string
+			// match that so we don't change the DDC key
+
+			return TEXT("");
+		}
+	}
 
 	virtual void GetSupportedFormats(TArray<FName>& OutFormats) const override
 	{
