TextureFormatUncompressed remove unnecessary casts

diff --git a/Engine/Source/Developer/TextureFormatUncompressed/Private/TextureFormatUncompressed.cpp b/Engine/Source/Developer/TextureFormatUncompressed/Private/TextureFormatUncompressed.cpp
--- a/Engine/Source/Developer/TextureFormatUncompressed/Private/TextureFormatUncompressed.cpp
+++ b/Engine/Source/Developer/TextureFormatUncompressed/Private/TextureFormatUncompressed.cpp
@@ -211,14 +211,14 @@ class FTextureFormatUncompressed : public ITextureFormat
 			OutCompressedImage.RawData.AddUninitialized(NumTexels * 4);
 			const FColor* FirstColor = (&Image.AsBGRA8()[0]);
 			const FColor* LastColor = FirstColor + NumTexels;
-			int8* Dest = (int8*)OutCompressedImage.RawData.GetData();
+			uint8* Dest = OutCompressedImage.RawData.GetData();
 
 			for (const FColor* Color = FirstColor; Color < LastColor; ++Color)
 			{
-				*Dest++ = (int32)Color->R;
-				*Dest++ = (int32)Color->G;
-				*Dest++ = (int32)Color->B;
-				*Dest++ = (int32)Color->A;
+				*Dest++ = Color->R;
+				*Dest++ = Color->G;
+				*Dest++ = Color->B;
+				*Dest++ = Color->A;
 			}
 
 			return true;
@@ -238,14 +238,14 @@ class FTextureFormatUncompressed : public ITextureFormat
 			OutCompressedImage.RawData.AddUninitialized(NumTexels * 4);
 			const FColor* FirstColor = (&Image.AsBGRA8()[0]);
 			const FColor* LastColor = FirstColor + NumTexels;
-			int8* Dest = (int8*)OutCompressedImage.RawData.GetData();
+			uint8* Dest = OutCompressedImage.RawData.GetData();
 
 			for (const FColor* Color = FirstColor; Color < LastColor; ++Color)
 			{
-				*Dest++ = (int32)Color->B;
-				*Dest++ = (int32)Color->G;
-				*Dest++ = (int32)Color->A;
-				*Dest++ = (int32)Color->R;
+				*Dest++ = Color->B;
+				*Dest++ = Color->G;
+				*Dest++ = Color->A;
+				*Dest++ = Color->R;
 			}
 
 			return true;
@@ -294,7 +294,7 @@ class FTextureFormatUncompressed : public ITextureFormat
 
 			// write out texels
 			uint8* Src = ErrorData.GetData();
-			uint8* Dest = (uint8*)OutCompressedImage.RawData.GetData();
+			uint8* Dest = OutCompressedImage.RawData.GetData();
 			for (int32 Y = 0; Y < InImage.SizeY; Y++)
 			{
 				for (int32 X = 0; X < InImage.SizeX * 4; X++)
