add a new TextureCompressionSetting TC_LQ(BGR565/BGR555A1), which is needed in Runtime virtual texture's baked SVT
The endian of BGR555A1 texture is troublesome on Android platform, because Opengl ES doesn't support GL_UNSIGNED_SHORT_1_5_5_5_REV texture format, meanwhile VK_FORMAT_R5G5B5A1_UNORM_PACK16 is not supported on some Android devices.

diff --git a/Engine/Source/Developer/Android/AndroidTargetPlatform/Private/AndroidTargetPlatform.cpp b/Engine/Source/Developer/Android/AndroidTargetPlatform/Private/AndroidTargetPlatform.cpp
--- a/Engine/Source/Developer/Android/AndroidTargetPlatform/Private/AndroidTargetPlatform.cpp
+++ b/Engine/Source/Developer/Android/AndroidTargetPlatform/Private/AndroidTargetPlatform.cpp
@@ -460,6 +460,10 @@ void FAndroidTargetPlatform::GetTextureFormats( const UTexture* InTexture, TArra
 			{
 				FormatPerLayer[LayerIndex] = AndroidTexFormat::NameBGRA8;
 			}
+			else if (LayerFormatSettings.CompressionSettings == TC_LQ)
+			{
+				FormatPerLayer[LayerIndex] = LayerFormatSettings.CompressionNoAlpha ? AndroidTexFormat::NameR5G6B5 : AndroidTexFormat::NameRGB555A1;
+			}
 			else if (LayerFormatSettings.CompressionSettings == TC_EncodedReflectionCapture && !LayerFormatSettings.CompressionNone)
 			{
 				FormatPerLayer[LayerIndex] = AndroidTexFormat::NameETC2_RGBA;
diff --git a/Engine/Source/Developer/Android/AndroidTargetPlatform/Private/AndroidTargetPlatform.h b/Engine/Source/Developer/Android/AndroidTargetPlatform/Private/AndroidTargetPlatform.h
--- a/Engine/Source/Developer/Android/AndroidTargetPlatform/Private/AndroidTargetPlatform.h
+++ b/Engine/Source/Developer/Android/AndroidTargetPlatform/Private/AndroidTargetPlatform.h
@@ -66,7 +66,10 @@ namespace AndroidTexFormat
 	static FName NameVU8(TEXT("VU8"));
 	static FName NameRGBA16F(TEXT("RGBA16F"));
 	static FName NameR16F(TEXT("R16F"));
-
+	static FName NameR5G6B5(TEXT("R5G6B5"));
+	static FName NameA1RGB555(TEXT("A1RGB555"));
+	//A1RGB555 is mapped to RGB555A1, because OpenGL GL_RGB5_A1 only supports alpha on the lowest bit
+	static FName NameRGB555A1(TEXT("RGB555A1"));
 	// Error "formats" (uncompressed)
 	static FName NamePOTERROR(TEXT("POTERROR"));
 }
@@ -348,6 +351,7 @@ public:
 			{ { FName(TEXT("BC6H")) },		{ FName(TEXT("ASTC_RGB")) } },
 			{ { FName(TEXT("BC7")) },		{ FName(TEXT("ASTC_RGBAuto")) } },
 			{ { FName(TEXT("AutoDXT")) },	{ FName(TEXT("ASTC_RGBAuto")) } },
+			{ { AndroidTexFormat::NameA1RGB555 },	{ AndroidTexFormat::NameRGB555A1 } },
 		};
 
 		// Supported in ES3.2 with ASTC
diff --git a/Engine/Source/Developer/Mac/MacTargetPlatform/Private/GenericMacTargetPlatform.h b/Engine/Source/Developer/Mac/MacTargetPlatform/Private/GenericMacTargetPlatform.h
--- a/Engine/Source/Developer/Mac/MacTargetPlatform/Private/GenericMacTargetPlatform.h
+++ b/Engine/Source/Developer/Mac/MacTargetPlatform/Private/GenericMacTargetPlatform.h
@@ -300,6 +300,8 @@ return TSuper::SupportsFeature(Feature);
 		return TextureFormatName;
 	}
 
+	virtual bool SupportsLQCompressionTextureFormat() const override { return false; };
+
 	virtual const UTextureLODSettings& GetTextureLODSettings() const override
 	{
 		return *TextureLODSettings;
diff --git a/Engine/Source/Developer/TargetPlatform/Public/Common/TargetPlatformBase.h b/Engine/Source/Developer/TargetPlatform/Public/Common/TargetPlatformBase.h
--- a/Engine/Source/Developer/TargetPlatform/Public/Common/TargetPlatformBase.h
+++ b/Engine/Source/Developer/TargetPlatform/Public/Common/TargetPlatformBase.h
@@ -90,6 +90,9 @@ public:
 	{
 		return FName();
 	}
+
+	virtual bool SupportsLQCompressionTextureFormat() const override { return true; };
+
 #endif //WITH_ENGINE
 
 	virtual bool PackageBuild( const FString& InPackgeDirectory ) override
diff --git a/Engine/Source/Developer/TargetPlatform/Public/Interfaces/ITargetPlatform.h b/Engine/Source/Developer/TargetPlatform/Public/Interfaces/ITargetPlatform.h
--- a/Engine/Source/Developer/TargetPlatform/Public/Interfaces/ITargetPlatform.h
+++ b/Engine/Source/Developer/TargetPlatform/Public/Interfaces/ITargetPlatform.h
@@ -540,7 +540,10 @@ public:
 		int32 SourceFormat,
 		bool bAllowCompression, bool bNoAlpha,
 		bool bSupportDX11TextureFormats, int32 TextureCompressionSettings) const = 0;
-
+	
+	//whether R5G6B5 and B5G5R5A1 is supported
+	virtual bool SupportsLQCompressionTextureFormat() const = 0;
+	
 	/**
 	 * Gets the format to use for a particular piece of audio.
 	 *
diff --git a/Engine/Source/Developer/TextureFormatUncompressed/Private/TextureFormatUncompressed.cpp b/Engine/Source/Developer/TextureFormatUncompressed/Private/TextureFormatUncompressed.cpp
--- a/Engine/Source/Developer/TextureFormatUncompressed/Private/TextureFormatUncompressed.cpp
+++ b/Engine/Source/Developer/TextureFormatUncompressed/Private/TextureFormatUncompressed.cpp
@@ -37,8 +37,10 @@ class FUncompressedTextureBuildFunction final : public FTextureBuildFunction
 	op(XGXR8) \
 	op(RGBA8) \
 	op(POTERROR) \
-	op(R16F)
-
+	op(R16F) \
+	op(R5G6B5) \
+	op(A1RGB555) \
+	op(RGB555A1)
 #define DECL_FORMAT_NAME(FormatName) static FName GTextureFormatName##FormatName = FName(TEXT(#FormatName));
 ENUM_SUPPORTED_FORMATS(DECL_FORMAT_NAME);
 #undef DECL_FORMAT_NAME
@@ -133,6 +135,14 @@ class FTextureFormatUncompressed : public ITextureFormat
 		{
 			return PF_B8G8R8A8;
 		}
+		else if (BuildSettings.TextureFormatName == GTextureFormatNameR5G6B5)
+		{
+			return PF_R5G6B5_UNORM;
+		}
+		else if (BuildSettings.TextureFormatName == GTextureFormatNameA1RGB555 || BuildSettings.TextureFormatName == GTextureFormatNameRGB555A1)
+		{
+			return PF_B5G5R5A1_UNORM;
+		}
 
 		UE_LOG(LogTextureFormatUncompressed, Fatal, TEXT("Unhandled texture format '%s' given to FTextureFormatUncompressed::GetPixelFormatForImage()"), *BuildSettings.TextureFormatName.ToString());
 		return PF_Unknown;
@@ -319,6 +329,51 @@ class FTextureFormatUncompressed : public ITextureFormat
 			
 			return true;
 		}
+		else if (BuildSettings.TextureFormatName == GTextureFormatNameR5G6B5 || BuildSettings.TextureFormatName == GTextureFormatNameRGB555A1 || BuildSettings.TextureFormatName == GTextureFormatNameA1RGB555)
+		{
+			FImage Image;
+			InImage.CopyTo(Image, ERawImageFormat::BGRA8, BuildSettings.GetGammaSpace());
+
+			OutCompressedImage.SizeX = Image.SizeX;
+			OutCompressedImage.SizeY = Image.SizeY;
+			OutCompressedImage.SizeZ = (BuildSettings.bVolume || BuildSettings.bTextureArray) ? Image.NumSlices : 1;
+
+			// swizzle each texel
+			uint64 NumTexels = (uint64)Image.SizeX * Image.SizeY * Image.NumSlices;
+			OutCompressedImage.RawData.Empty(NumTexels * 2);
+			OutCompressedImage.RawData.AddUninitialized(NumTexels * 2);
+			const FColor* FirstColor = (&Image.AsBGRA8()[0]);
+			const FColor* LastColor = FirstColor + NumTexels;
+			uint16* Dest = (uint16*)OutCompressedImage.RawData.GetData();
+
+			if(BuildSettings.TextureFormatName == GTextureFormatNameR5G6B5)
+			{
+				for (const FColor* Color = FirstColor; Color < LastColor; ++Color)
+				{
+					uint16 BGR565 = (uint16(Color->R >> 3) << 11) | (uint16(Color->G >> 2) << 5) | uint16(Color->B >> 3);
+					*Dest++ = BGR565;
+				}
+			}
+			else if(BuildSettings.TextureFormatName == GTextureFormatNameA1RGB555)
+			{
+				for (const FColor* Color = FirstColor; Color < LastColor; ++Color)
+				{
+					//most rhi supports alpha on the highest bit
+					uint16 BGR555A1 = (uint16(Color->A >> 7) << 15) | (uint16(Color->R >> 3) << 10) | (uint16(Color->G >> 3) << 5) | uint16(Color->B >> 3);
+					*Dest++ = BGR555A1;
+				}
+			}
+			else
+			{
+				for (const FColor* Color = FirstColor; Color < LastColor; ++Color)
+				{
+					//OpenGL GL_RGB5_A1 only supports alpha on the lowest bit
+					uint16 BGR555A1 = (uint16(Color->R >> 3) << 11) | (uint16(Color->G >> 3) << 6) | (uint16(Color->B >> 3) << 1) | uint16(Color->A >> 7);
+					*Dest++ = BGR555A1;
+				}
+			}
+			return true;
+		}
 
 		UE_LOG(LogTextureFormatUncompressed, Warning,
 			TEXT("Cannot convert uncompressed image to format '%s'."),
diff --git a/Engine/Source/Runtime/Engine/Classes/Engine/TextureDefines.h b/Engine/Source/Runtime/Engine/Classes/Engine/TextureDefines.h
--- a/Engine/Source/Runtime/Engine/Classes/Engine/TextureDefines.h
+++ b/Engine/Source/Runtime/Engine/Classes/Engine/TextureDefines.h
@@ -290,6 +290,7 @@ enum TextureCompressionSettings
 	TC_HDR_Compressed			UMETA(DisplayName = "HDRCompressed (RGB, BC6H, DX11)"),
 	TC_BC7						UMETA(DisplayName = "BC7 (DX11, optional A)"),
 	TC_HalfFloat				UMETA(DisplayName = "Half Float (R16F)"),
+	TC_LQ				        UMETA(Hidden, DisplayName = "Low Quality (BGR565/BGR555A1)", ToolTip = "BGR565/BGR555A1, fallback to DXT1/DXT5 on Mac platform"),
 	TC_EncodedReflectionCapture		UMETA(Hidden),
 	TC_MAX,
 };
diff --git a/Engine/Source/Runtime/Engine/Private/Texture.cpp b/Engine/Source/Runtime/Engine/Private/Texture.cpp
--- a/Engine/Source/Runtime/Engine/Private/Texture.cpp
+++ b/Engine/Source/Runtime/Engine/Private/Texture.cpp
@@ -1890,6 +1890,9 @@ FName GetDefaultTextureFormatName( const ITargetPlatform* TargetPlatform, const
 	static FName NameR16F(TEXT("R16F"));
 	static FName NameBC6H(TEXT("BC6H"));
 	static FName NameBC7(TEXT("BC7"));
+	static FName NameR5G6B5(TEXT("R5G6B5"));
+	static FName NameA1RGB555(TEXT("A1RGB555"));
+	
 
 	check(TargetPlatform);
 
@@ -1957,6 +1960,18 @@ FName GetDefaultTextureFormatName( const ITargetPlatform* TargetPlatform, const
 			TextureFormatName = NameBGRA8;
 		}
 	}
+	else if (FormatSettings.CompressionSettings == TC_LQ)
+	{
+		bool bLQCompressionSupported = TargetPlatform->SupportsLQCompressionTextureFormat();
+		if(bLQCompressionSupported)
+		{
+			TextureFormatName = FormatSettings.CompressionNoAlpha ? NameR5G6B5 : NameA1RGB555;
+		}
+		else
+		{
+			TextureFormatName = FormatSettings.CompressionNoAlpha ? NameDXT1 : NameDXT5;
+		}
+	}
 	else if (FormatSettings.CompressionSettings == TC_HDR)
 	{
 		TextureFormatName = NameRGBA16F;
diff --git a/Engine/Source/Runtime/OpenGLDrv/Private/OpenGLDevice.cpp b/Engine/Source/Runtime/OpenGLDrv/Private/OpenGLDevice.cpp
--- a/Engine/Source/Runtime/OpenGLDrv/Private/OpenGLDevice.cpp
+++ b/Engine/Source/Runtime/OpenGLDrv/Private/OpenGLDevice.cpp
@@ -927,7 +927,13 @@ static void InitRHICapabilitiesForGL()
 	SetupTextureFormat( PF_R8,					FOpenGLTextureFormat( GL_R8,					GL_R8,					GL_RED,				GL_UNSIGNED_BYTE,				false,			false));
 
 	SetupTextureFormat( PF_R5G6B5_UNORM,        FOpenGLTextureFormat( GL_RGB565,                GL_RGB565,              GL_RGB,             GL_UNSIGNED_SHORT_5_6_5,        false,          false));
-	SetupTextureFormat( PF_B5G5R5A1_UNORM,		FOpenGLTextureFormat( GL_RGB5_A1,               GL_RGB5_A1,             GL_RGBA,            GL_UNSIGNED_SHORT_5_5_5_1,      false,          false));
+
+	//GL_UNSIGNED_SHORT_1_5_5_5_REV is not defined in gles
+	GLenum GL5551Format = FOpenGL::GetPlatfrom5551Format();
+	//the last 5 bits of GL_UNSIGNED_SHORT_1_5_5_5_REV is Red, while the last 5 bits of DXGI_FORMAT_B5G5R5A1_UNORM is Blue
+	bool bNeedsToSwizzleRedBlue = GL5551Format != GL_UNSIGNED_SHORT_5_5_5_1;
+	SetupTextureFormat( PF_B5G5R5A1_UNORM,		FOpenGLTextureFormat( GL_RGB5_A1,               GL_RGB5_A1,             GL_RGBA,			GL5551Format,					false,          bNeedsToSwizzleRedBlue));
+
 
 	SetupTextureFormat( PF_G16,					FOpenGLTextureFormat( GL_R16,					GL_R16,					GL_RED,				GL_UNSIGNED_SHORT,					false,	false));
 	SetupTextureFormat( PF_R32_FLOAT,			FOpenGLTextureFormat( GL_R32F,					GL_R32F,				GL_RED,				GL_FLOAT,							false,	false));
diff --git a/Engine/Source/Runtime/OpenGLDrv/Private/OpenGLTexture.cpp b/Engine/Source/Runtime/OpenGLDrv/Private/OpenGLTexture.cpp
--- a/Engine/Source/Runtime/OpenGLDrv/Private/OpenGLTexture.cpp
+++ b/Engine/Source/Runtime/OpenGLDrv/Private/OpenGLTexture.cpp
@@ -510,7 +510,13 @@ void FOpenGLDynamicRHI::InitializeGLTextureInternal(GLuint TextureID, FRHITextur
 
 	// Use a texture stage that's not likely to be used for draws, to avoid waiting
 	CachedSetupTextureStage(ContextState, FOpenGL::GetMaxCombinedTextureImageUnits() - 1, Target, TextureID, 0, NumMips);
-
+#if !PLATFORM_ANDROID
+	if(GLFormat.InternalFormat[0] == GL_RGB5_A1)
+	{
+		glTexParameteri(Target, GL_TEXTURE_SWIZZLE_R, GL_BLUE);
+		glTexParameteri(Target, GL_TEXTURE_SWIZZLE_B, GL_RED);
+	}
+#endif
 	if (NumSamples == 1 && !TileMemDepth)
 	{
 		if (Target == GL_TEXTURE_EXTERNAL_OES || !FMath::IsPowerOfTwo(SizeX) || !FMath::IsPowerOfTwo(SizeY))
diff --git a/Engine/Source/Runtime/OpenGLDrv/Public/OpenGL.h b/Engine/Source/Runtime/OpenGLDrv/Public/OpenGL.h
--- a/Engine/Source/Runtime/OpenGLDrv/Public/OpenGL.h
+++ b/Engine/Source/Runtime/OpenGLDrv/Public/OpenGL.h
@@ -366,6 +366,7 @@ public:
 	static bool PE_GLSLToDeviceCompatibleGLSL(FAnsiCharArray& GlslCodeOriginal, const FString& ShaderName, GLenum TypeEnum, const FOpenGLShaderDeviceCapabilities& Capabilities, FAnsiCharArray& GlslCode) UGL_OPTIONAL(false)
 	static void PE_SetupTextureFormat(void(*SetupTextureFormat)(EPixelFormat, const FOpenGLTextureFormat&)) UGL_OPTIONAL_VOID
 
+	static GLenum GetPlatfrom5551Format() UGL_OPTIONAL(GL_UNSIGNED_SHORT_5_5_5_1)
 protected:
 	static GLint MaxTextureImageUnits;
 	static GLint MaxCombinedTextureImageUnits;
diff --git a/Engine/Source/Runtime/OpenGLDrv/Public/OpenGL3.h b/Engine/Source/Runtime/OpenGLDrv/Public/OpenGL3.h
--- a/Engine/Source/Runtime/OpenGLDrv/Public/OpenGL3.h
+++ b/Engine/Source/Runtime/OpenGLDrv/Public/OpenGL3.h
@@ -771,6 +771,7 @@ struct FOpenGL3 : public FOpenGLBase
 
 	static FORCEINLINE bool IsAndroidGLESCompatibilityModeEnabled() { return bAndroidGLESCompatibilityMode; }
 
+	static GLenum GetPlatfrom5551Format() { return GL_UNSIGNED_SHORT_1_5_5_5_REV; }
 protected:
 	static GLsizei NextTextureName;
 	static GLuint TextureNamesCache[OPENGL_NAME_CACHE_SIZE];
diff --git a/Engine/Source/Runtime/VulkanRHI/Private/Android/VulkanAndroidPlatform.h b/Engine/Source/Runtime/VulkanRHI/Private/Android/VulkanAndroidPlatform.h
--- a/Engine/Source/Runtime/VulkanRHI/Private/Android/VulkanAndroidPlatform.h
+++ b/Engine/Source/Runtime/VulkanRHI/Private/Android/VulkanAndroidPlatform.h
@@ -129,6 +129,12 @@ public:
 
 	static void DestroySwapchainKHR(VkDevice Device, VkSwapchainKHR Swapchain, const VkAllocationCallbacks* Allocator);
 
+	static VkFormat GetPlatform5551FormatWithFallback(VkFormat& OutFallbackFormat0, VkFormat& OutFallbackFormat1) 
+	{ 
+		OutFallbackFormat0 = VK_FORMAT_A1R5G5B5_UNORM_PACK16; 
+		OutFallbackFormat1 = VK_FORMAT_B8G8R8A8_UNORM;
+		return VK_FORMAT_R5G5B5A1_UNORM_PACK16; 
+	};
 protected:
 	static void* VulkanLib;
 	static bool bAttemptedLoad;
diff --git a/Engine/Source/Runtime/VulkanRHI/Private/VulkanDevice.cpp b/Engine/Source/Runtime/VulkanRHI/Private/VulkanDevice.cpp
--- a/Engine/Source/Runtime/VulkanRHI/Private/VulkanDevice.cpp
+++ b/Engine/Source/Runtime/VulkanRHI/Private/VulkanDevice.cpp
@@ -629,7 +629,12 @@ void FVulkanDevice::SetupFormats()
 	MapFormatSupport(PF_R5G6B5_UNORM, VK_FORMAT_R5G6B5_UNORM_PACK16);
 	SetComponentMapping(PF_R5G6B5_UNORM, VK_COMPONENT_SWIZZLE_R, VK_COMPONENT_SWIZZLE_G, VK_COMPONENT_SWIZZLE_B, VK_COMPONENT_SWIZZLE_A);
 
-	MapFormatSupport(PF_B5G5R5A1_UNORM, VK_FORMAT_A1R5G5B5_UNORM_PACK16);
+	VkFormat Vulkan5551Format, FallbackFormat0, FallbackFormat1;
+	Vulkan5551Format = FVulkanPlatform::GetPlatform5551Format(FallbackFormat0, FallbackFormat1);
+
+	//use rgb555a1 on Android in order to match OpenGL GL_RGB5_A1
+	MapFormatSupportWithFallback(PF_B5G5R5A1_UNORM, Vulkan5551Format, { FallbackFormat0,  FallbackFormat1 });
+
 	SetComponentMapping(PF_B5G5R5A1_UNORM, VK_COMPONENT_SWIZZLE_R, VK_COMPONENT_SWIZZLE_G, VK_COMPONENT_SWIZZLE_B, VK_COMPONENT_SWIZZLE_A);
 
 	MapFormatSupport(PF_R8G8B8A8, VK_FORMAT_R8G8B8A8_UNORM);
diff --git a/Engine/Source/Runtime/VulkanRHI/Private/VulkanGenericPlatform.h b/Engine/Source/Runtime/VulkanRHI/Private/VulkanGenericPlatform.h
--- a/Engine/Source/Runtime/VulkanRHI/Private/VulkanGenericPlatform.h
+++ b/Engine/Source/Runtime/VulkanRHI/Private/VulkanGenericPlatform.h
@@ -126,4 +126,11 @@ public:
 
 	// Return VK_FALSE if platform wants to suppress the given debug report from the validation layers, VK_TRUE to print it.
 	static VkBool32 DebugReportFunction(VkDebugReportFlagsEXT MsgFlags, VkDebugReportObjectTypeEXT ObjType, uint64_t SrcObject, size_t Location, int32 MsgCode, const ANSICHAR* LayerPrefix, const ANSICHAR* Msg, void* UserData) { return VK_TRUE; }
+
+	static VkFormat GetPlatform5551Format(VkFormat& OutFallbackFormat0, VkFormat& OutFallbackFormat1)
+	{
+		OutFallbackFormat0 = VK_FORMAT_R5G5B5A1_UNORM_PACK16;
+		OutFallbackFormat1 = VK_FORMAT_B8G8R8A8_UNORM;
+		return VK_FORMAT_A1R5G5B5_UNORM_PACK16;
+	};
 };
