TextureFormatUncompressed CanAcceptNonF32Source and avoid alloc and copy when possible
also note that the 565 and 555 color conversions are wrong but do not change them
output unchanged, no DDC key bump

diff --git a/Engine/Source/Developer/TextureFormatUncompressed/Private/TextureFormatUncompressed.cpp b/Engine/Source/Developer/TextureFormatUncompressed/Private/TextureFormatUncompressed.cpp
--- a/Engine/Source/Developer/TextureFormatUncompressed/Private/TextureFormatUncompressed.cpp
+++ b/Engine/Source/Developer/TextureFormatUncompressed/Private/TextureFormatUncompressed.cpp
@@ -115,10 +115,7 @@ class FTextureFormatUncompressed : public ITextureFormat
 
 	virtual void GetSupportedFormats(TArray<FName>& OutFormats) const override
 	{
-		for (int32 i = 0; i < UE_ARRAY_COUNT(GSupportedTextureFormatNames); ++i)
-		{
-			OutFormats.Add(GSupportedTextureFormatNames[i]);
-		}
+		OutFormats.Append(GSupportedTextureFormatNames, UE_ARRAY_COUNT(GSupportedTextureFormatNames) ); 
 	}
 	
 	virtual EPixelFormat GetEncodedPixelFormat(const FTextureBuildSettings& BuildSettings, bool bImageHasAlphaChannel) const override
@@ -172,7 +169,40 @@ class FTextureFormatUncompressed : public ITextureFormat
 		UE_LOG(LogTextureFormatUncompressed, Fatal, TEXT("Unhandled texture format '%s' given to FTextureFormatUncompressed::GetEncodedPixelFormat()"), *BuildSettings.TextureFormatName.ToString());
 		return PF_Unknown;
 	}
+	
+	virtual bool CanAcceptNonF32Source(FName Format) const override
+	{
+		return true;
+	}
 
+	// InImage arg is not actually const
+	static bool DoCompressImageSimple(const FImage& InImage, ERawImageFormat::Type DestFormat, EGammaSpace DestGammaSpace, FCompressedImage2D& OutCompressedImage)
+	{
+		FImage Image;
+
+		if ( InImage.Format == DestFormat && InImage.GetGammaSpace() == DestGammaSpace )
+		{
+			// source image is already in desired format, just move it
+			// source image can be mutated by compressor
+			Image.Swap(const_cast<FImage &>(InImage));
+			
+			// keep InImage info unchanged, only free the RawData bits
+			(FImageInfo &)(const_cast<FImage &>(InImage)) = (FImageInfo &)Image;
+
+			// note that the source bits are not freed; they're still held allocated in the OutCompressedImage
+			//	currently freeing the source is not allowed because it may be being hashed on a thread
+		}
+		else
+		{
+			InImage.CopyTo(Image, DestFormat, DestGammaSpace);
+		}
+
+		OutCompressedImage.RawData = MoveTemp(Image.RawData);
+
+		return true;
+	}
+	
+	// InImage arg is not actually const
 	virtual bool CompressImage(
 		const FImage& InImage,
 		const FTextureBuildSettings& BuildSettings,
@@ -185,47 +215,31 @@ class FTextureFormatUncompressed : public ITextureFormat
 		FCompressedImage2D& OutCompressedImage
 		) const override
 	{
-		OutCompressedImage.PixelFormat = GetEncodedPixelFormat(BuildSettings, bImageHasAlphaChannel);
+		TRACE_CPUPROFILER_EVENT_SCOPE(TFUncompressed.CompressImage);
 
-		// @todo Oodle : fix me : lots of pointless code dupe here
-		//	most of these should just map the Name to an ERawImageFormat and do CopyImage
+		// InImage can be any format because CanAcceptNonF32Source
+
+		OutCompressedImage.PixelFormat = GetEncodedPixelFormat(BuildSettings, bImageHasAlphaChannel);
+		OutCompressedImage.SizeX = InImage.SizeX;
+		OutCompressedImage.SizeY = InImage.SizeY;
+		// note: cubes come in as 6 slices and go out as 1
+		OutCompressedImage.SizeZ = (BuildSettings.bVolume || BuildSettings.bTextureArray) ? InImage.NumSlices : 1;
 
 		if (BuildSettings.TextureFormatName == GTextureFormatNameG8)
 		{
-			FImage Image;
-			InImage.CopyTo(Image, ERawImageFormat::G8, BuildSettings.GetDestGammaSpace());
-
-			OutCompressedImage.SizeX = Image.SizeX;
-			OutCompressedImage.SizeY = Image.SizeY;
-			OutCompressedImage.SizeZ = (BuildSettings.bVolume || BuildSettings.bTextureArray) ? Image.NumSlices : 1;
-			OutCompressedImage.RawData = MoveTemp(Image.RawData);
-
-			return true;
+			return DoCompressImageSimple(InImage, ERawImageFormat::G8, BuildSettings.GetDestGammaSpace(), OutCompressedImage);
 		}
 		else if (BuildSettings.TextureFormatName == GTextureFormatNameG16)
 		{
-			FImage Image;
-			InImage.CopyTo(Image, ERawImageFormat::G16, EGammaSpace::Linear);
-
-			OutCompressedImage.SizeX = Image.SizeX;
-			OutCompressedImage.SizeY = Image.SizeY;
-			OutCompressedImage.SizeZ = BuildSettings.bVolume ? Image.NumSlices : 1;
-			OutCompressedImage.RawData = MoveTemp(Image.RawData);
-
-			return true;
+			return DoCompressImageSimple(InImage, ERawImageFormat::G16, EGammaSpace::Linear, OutCompressedImage);
 		}
 		else if (BuildSettings.TextureFormatName == GTextureFormatNameVU8)
 		{
 			FImage Image;
 			InImage.CopyTo(Image, ERawImageFormat::BGRA8, BuildSettings.GetDestGammaSpace());
 
-			OutCompressedImage.SizeX = Image.SizeX;
-			OutCompressedImage.SizeY = Image.SizeY;
-			OutCompressedImage.SizeZ = (BuildSettings.bVolume || BuildSettings.bTextureArray) ? Image.NumSlices : 1;
-
-			uint64 NumTexels = (uint64)Image.SizeX * Image.SizeY * Image.NumSlices;
-			OutCompressedImage.RawData.Empty(NumTexels * 2);
-			OutCompressedImage.RawData.AddUninitialized(NumTexels * 2);
+			int64 NumTexels = Image.GetNumPixels();
+			OutCompressedImage.RawData.SetNumUninitialized(NumTexels * 2);
 			const FColor* FirstColor = (&Image.AsBGRA8()[0]);
 			const FColor* LastColor = FirstColor + NumTexels;
 			int8* Dest = (int8*)OutCompressedImage.RawData.GetData();
@@ -240,41 +254,16 @@ class FTextureFormatUncompressed : public ITextureFormat
 		}
 		else if (BuildSettings.TextureFormatName == GTextureFormatNameBGRA8)
 		{
-			FImage Image;
-			InImage.CopyTo(Image, ERawImageFormat::BGRA8, BuildSettings.GetDestGammaSpace());
-
-			OutCompressedImage.SizeX = Image.SizeX;
-			OutCompressedImage.SizeY = Image.SizeY;
-			OutCompressedImage.SizeZ = (BuildSettings.bVolume || BuildSettings.bTextureArray) ? Image.NumSlices : 1;
-			OutCompressedImage.RawData = MoveTemp(Image.RawData);
-
-			return true;
+			return DoCompressImageSimple(InImage, ERawImageFormat::BGRA8, BuildSettings.GetDestGammaSpace(), OutCompressedImage);
 		}
 		else if (BuildSettings.TextureFormatName == GTextureFormatNameRGBA8)
 		{
 			FImage Image;
 			InImage.CopyTo(Image, ERawImageFormat::BGRA8, BuildSettings.GetDestGammaSpace());
+			
+			FImageCore::TransposeImageRGBABGRA(Image);
 
-			OutCompressedImage.SizeX = Image.SizeX;
-			OutCompressedImage.SizeY = Image.SizeY;
-			OutCompressedImage.SizeZ = (BuildSettings.bVolume || BuildSettings.bTextureArray) ? Image.NumSlices : 1;
-
-			// BGRA to RGBA : @todo Oodle : just use ImageCore CopyImageRGBABGRA
-			// swizzle each texel
-			uint64 NumTexels = (uint64)Image.SizeX * Image.SizeY * Image.NumSlices;
-			OutCompressedImage.RawData.Empty(NumTexels * 4);
-			OutCompressedImage.RawData.AddUninitialized(NumTexels * 4);
-			const FColor* FirstColor = (&Image.AsBGRA8()[0]);
-			const FColor* LastColor = FirstColor + NumTexels;
-			uint8* Dest = OutCompressedImage.RawData.GetData();
-
-			for (const FColor* Color = FirstColor; Color < LastColor; ++Color)
-			{
-				*Dest++ = Color->R;
-				*Dest++ = Color->G;
-				*Dest++ = Color->B;
-				*Dest++ = Color->A;
-			}
+			OutCompressedImage.RawData = MoveTemp(Image.RawData);
 
 			return true;
 		}
@@ -283,20 +272,16 @@ class FTextureFormatUncompressed : public ITextureFormat
 			FImage Image;
 			InImage.CopyTo(Image, ERawImageFormat::BGRA8, BuildSettings.GetDestGammaSpace());
 
-			OutCompressedImage.SizeX = Image.SizeX;
-			OutCompressedImage.SizeY = Image.SizeY;
-			OutCompressedImage.SizeZ = (BuildSettings.bVolume || BuildSettings.bTextureArray) ? Image.NumSlices : 1;
-
 			// swizzle each texel
-			uint64 NumTexels = (uint64)Image.SizeX * Image.SizeY * Image.NumSlices;
-			OutCompressedImage.RawData.Empty(NumTexels * 4);
-			OutCompressedImage.RawData.AddUninitialized(NumTexels * 4);
+			int64 NumTexels = Image.GetNumPixels();
+			OutCompressedImage.RawData.SetNumUninitialized(NumTexels * 4);
 			const FColor* FirstColor = (&Image.AsBGRA8()[0]);
 			const FColor* LastColor = FirstColor + NumTexels;
 			uint8* Dest = OutCompressedImage.RawData.GetData();
 
 			for (const FColor* Color = FirstColor; Color < LastColor; ++Color)
 			{
+				// XGXR8
 				*Dest++ = Color->B;
 				*Dest++ = Color->G;
 				*Dest++ = Color->A;
@@ -307,64 +292,29 @@ class FTextureFormatUncompressed : public ITextureFormat
 		}
 		else if (BuildSettings.TextureFormatName == GTextureFormatNameRGBA16F)
 		{
-			FImage Image;
-			InImage.CopyTo(Image, ERawImageFormat::RGBA16F, EGammaSpace::Linear);
-
-			OutCompressedImage.SizeX = Image.SizeX;
-			OutCompressedImage.SizeY = Image.SizeY;
-			OutCompressedImage.SizeZ = (BuildSettings.bVolume || BuildSettings.bTextureArray) ? Image.NumSlices : 1;
-			OutCompressedImage.RawData = MoveTemp(Image.RawData);
-
-			return true;
+			return DoCompressImageSimple(InImage, ERawImageFormat::RGBA16F, EGammaSpace::Linear, OutCompressedImage);
 		}
 		else if (BuildSettings.TextureFormatName == GTextureFormatNameRGBA32F)
 		{
-			FImage Image;
-			InImage.CopyTo(Image, ERawImageFormat::RGBA32F, EGammaSpace::Linear);
-
-			OutCompressedImage.SizeX = Image.SizeX;
-			OutCompressedImage.SizeY = Image.SizeY;
-			OutCompressedImage.SizeZ = (BuildSettings.bVolume || BuildSettings.bTextureArray) ? Image.NumSlices : 1;
-			OutCompressedImage.RawData = MoveTemp(Image.RawData);
-
-			return true;
+			return DoCompressImageSimple(InImage, ERawImageFormat::RGBA32F, EGammaSpace::Linear, OutCompressedImage);
 		}
 		else if (BuildSettings.TextureFormatName == GTextureFormatNameR16F)
 		{
-			FImage Image;
-			InImage.CopyTo(Image, ERawImageFormat::R16F, EGammaSpace::Linear);
-
-			OutCompressedImage.SizeX = Image.SizeX;
-			OutCompressedImage.SizeY = Image.SizeY;
-			OutCompressedImage.SizeZ = (BuildSettings.bVolume || BuildSettings.bTextureArray) ? Image.NumSlices : 1;
-			OutCompressedImage.RawData = MoveTemp(Image.RawData);
-
-			return true;
+			return DoCompressImageSimple(InImage, ERawImageFormat::R16F, EGammaSpace::Linear, OutCompressedImage);
 		}
 		else if (BuildSettings.TextureFormatName == GTextureFormatNameR32F)
 		{
-			FImage Image;
-			InImage.CopyTo(Image, ERawImageFormat::R32F, EGammaSpace::Linear);
-
-			OutCompressedImage.SizeX = Image.SizeX;
-			OutCompressedImage.SizeY = Image.SizeY;
-			OutCompressedImage.SizeZ = (BuildSettings.bVolume || BuildSettings.bTextureArray) ? Image.NumSlices : 1;
-			OutCompressedImage.RawData = MoveTemp(Image.RawData);
-
-			return true;
+			return DoCompressImageSimple(InImage, ERawImageFormat::R32F, EGammaSpace::Linear, OutCompressedImage);
 		}
 		else if (BuildSettings.TextureFormatName == GTextureFormatNamePOTERROR)
 		{
+			// @todo : this looks like temp debug code. please remove.
+
 			// load the error image data we will just repeat into the texture
 			TArray64<uint8> ErrorData;
 			FFileHelper::LoadFileToArray(ErrorData, *(FPaths::EngineDir() / TEXT("Content/MobileResources/PowerOfTwoError64x64.raw")));
 			check(ErrorData.Num() == 16384);
 
-			// set output
-			OutCompressedImage.SizeX = InImage.SizeX;
-			OutCompressedImage.SizeY = InImage.SizeY;
-			OutCompressedImage.SizeZ = (BuildSettings.bVolume || BuildSettings.bTextureArray) ? InImage.NumSlices : 1;
-
 			// allocate output memory
 			check(InImage.NumSlices == 1);
 			uint64 NumTexels = (uint64)InImage.SizeX * InImage.SizeY;
@@ -383,7 +333,6 @@ class FTextureFormatUncompressed : public ITextureFormat
 					Dest[(uint64)Y * InImage.SizeX * 4 + X] = Src[(uint64)SrcY * 64 * 4 + SrcX];
 				}
 			}
-
 			
 			return true;
 		}
@@ -392,18 +341,18 @@ class FTextureFormatUncompressed : public ITextureFormat
 			FImage Image;
 			InImage.CopyTo(Image, ERawImageFormat::BGRA8, BuildSettings.GetDestGammaSpace());
 
-			OutCompressedImage.SizeX = Image.SizeX;
-			OutCompressedImage.SizeY = Image.SizeY;
-			OutCompressedImage.SizeZ = (BuildSettings.bVolume || BuildSettings.bTextureArray) ? Image.NumSlices : 1;
-
 			// swizzle each texel
-			uint64 NumTexels = (uint64)Image.SizeX * Image.SizeY * Image.NumSlices;
+			int64 NumTexels = Image.GetNumPixels();
 			OutCompressedImage.RawData.Empty(NumTexels * 2);
 			OutCompressedImage.RawData.AddUninitialized(NumTexels * 2);
 			const FColor* FirstColor = (&Image.AsBGRA8()[0]);
 			const FColor* LastColor = FirstColor + NumTexels;
 			uint16* Dest = (uint16*)OutCompressedImage.RawData.GetData();
 
+			// @todo Oodle : this is not a correct 565 color quantizer.  See rrColor565Bits_Quantize.
+			//	 not a big deal because this is rarely used these days, but beware do not copy-paste and spread this mistake
+			//	(if you fix this, must bump DDC key)
+
 			if(BuildSettings.TextureFormatName == GTextureFormatNameR5G6B5)
 			{
 				for (const FColor* Color = FirstColor; Color < LastColor; ++Color)
@@ -432,13 +381,15 @@ class FTextureFormatUncompressed : public ITextureFormat
 			}
 			return true;
 		}
+		else
+		{
+			UE_LOG(LogTextureFormatUncompressed, Warning,
+				TEXT("Cannot convert uncompressed image to format '%s'."),
+				*BuildSettings.TextureFormatName.ToString()
+				);
 
-		UE_LOG(LogTextureFormatUncompressed, Warning,
-			TEXT("Cannot convert uncompressed image to format '%s'."),
-			*BuildSettings.TextureFormatName.ToString()
-			);
-
-		return false;
+			return false;
+		}
 	}
 };
 
